shutdown fix + remove dead code

#rnx

#Synced-CL 51007224

diff --git a/Engine/Plugins/Online/Microsoft/OnlineSubsystemGDK/Source/Private/OnlineSessionGDK.cpp b/Engine/Plugins/Online/Microsoft/OnlineSubsystemGDK/Source/Private/OnlineSessionGDK.cpp
--- a/Engine/Plugins/Online/Microsoft/OnlineSubsystemGDK/Source/Private/OnlineSessionGDK.cpp
+++ b/Engine/Plugins/Online/Microsoft/OnlineSubsystemGDK/Source/Private/OnlineSessionGDK.cpp
@@ -32,15 +32,10 @@ FOnlineSessionGDK::FOnlineSessionGDK(FOnlineSubsystemGDK* InSubsystem)
 
 FOnlineSessionGDK::~FOnlineSessionGDK()
 {
-#if 0
-	if (XGameRuntimeIsFeatureAvailable(XGameRuntimeFeature::XGameInvite))
+	if (IGDKRuntimeModule* GDKRuntimeModule = IGDKRuntimeModule::TryGet())
 	{
-		GDK_SCOPE_NOT_TIME_SENSITIVE(); // (shutdown only) XGameInviteUnregisterForEvent is not safe to call on time-sensitive threads
-		XGameInviteUnregisterForEvent(InviteAcceptedHandler, false);
+		GDKRuntimeModule->UnregisterGameInviteReceivedDelegate(InviteAcceptedHandler);
 	}
-#else
-	IGDKRuntimeModule::Get().UnregisterGameInviteReceivedDelegate(InviteAcceptedHandler);
-#endif
 }
 
 bool FOnlineSessionGDK::IsMpaEnabled() const
