Fixed GDK subsystem crash on quick logout due to invalid StoreContextHandle reference.

#Synced-CL 50981645

diff --git a/Engine/Plugins/Online/Microsoft/OnlineSubsystemGDK/Source/Private/AsyncTasks/OnlineAsyncTaskGDKQueryOffers.cpp b/Engine/Plugins/Online/Microsoft/OnlineSubsystemGDK/Source/Private/AsyncTasks/OnlineAsyncTaskGDKQueryOffers.cpp
--- a/Engine/Plugins/Online/Microsoft/OnlineSubsystemGDK/Source/Private/AsyncTasks/OnlineAsyncTaskGDKQueryOffers.cpp
+++ b/Engine/Plugins/Online/Microsoft/OnlineSubsystemGDK/Source/Private/AsyncTasks/OnlineAsyncTaskGDKQueryOffers.cpp
@@ -84,6 +84,14 @@ void FOnlineAsyncTaskGDKQueryOffers::StartQuery()
 		XStoreProductKind::UnmanagedConsumable;
 
 	XStoreContextHandle StoreContextHandle = Subsystem->GetStoreGDK()->GetStoreContextHandle(GDKUserHandle);
+	if (!StoreContextHandle)
+	{
+		ErrorResponse.bSucceeded = false;
+		ErrorResponse.SetFromErrorMessage(FText::FromString(TEXT("Error querying store offers ContextHandle.")));
+		bWasSuccessful = false;
+		bIsComplete = true;
+		return;
+	}
 
 	// From GDK documentation: Restricts the results by some action stored in the product document. By default, this API returns all
 	// products, even if they are not purchasable, but you can restrict this to "Purchase" if you only want purchasable, or "License"
