OnlineSubsystemGDK & PlayFabParty are now loaded on-demand
- prevents the modules from registering itself on startup even when they aren't wanted
- OnlineSubsystem module will load OnlineSubsystemGDK when configured, and that will load PlayFabParty when necessary.

#rnx

#Synced-CL 51338024

diff --git a/Engine/Plugins/Online/Microsoft/OnlineSubsystemGDK/OnlineSubsystemGDK.uplugin b/Engine/Plugins/Online/Microsoft/OnlineSubsystemGDK/OnlineSubsystemGDK.uplugin
--- a/Engine/Plugins/Online/Microsoft/OnlineSubsystemGDK/OnlineSubsystemGDK.uplugin
+++ b/Engine/Plugins/Online/Microsoft/OnlineSubsystemGDK/OnlineSubsystemGDK.uplugin
@@ -11,6 +11,7 @@
 	"EnabledByDefault" : false,
 	"RequiresBuildPlatform": false,
 	"HasExplicitPlatforms": true,	
+	"ExplicitlyLoaded":  true,
 	"SupportedTargetPlatforms": ["Win64"],
 
 	"Modules": [
diff --git a/Engine/Plugins/Online/Microsoft/OnlineSubsystemGDK/Source/Private/OnlineSubsystemGDK.cpp b/Engine/Plugins/Online/Microsoft/OnlineSubsystemGDK/Source/Private/OnlineSubsystemGDK.cpp
--- a/Engine/Plugins/Online/Microsoft/OnlineSubsystemGDK/Source/Private/OnlineSubsystemGDK.cpp
+++ b/Engine/Plugins/Online/Microsoft/OnlineSubsystemGDK/Source/Private/OnlineSubsystemGDK.cpp
@@ -25,6 +25,7 @@
 #include "OnlineSessionGDK.h"
 #include "OnlineSessionInterfaceMpsdGDK.h"
 #include "OnlinePresenceInterfaceGDK.h"
+#include "OnlineDelegates.h"
 #if WITH_ENGINE
 #include "OnlineSubsystemUtils.h"
 #include "OnlineVoiceInterfaceGDK.h"
@@ -425,6 +426,17 @@ bool FOnlineSubsystemGDK::Init()
 #if WITH_EDITOR
 		bIsInitialized = true;
 #endif
+
+#ifdef UE_PLAYFAB_MATCHMAKING
+		// once we have been registered, it is safe to load our dependent modules
+		OnOnlineSubsystemCreated = FOnlineSubsystemDelegates::OnOnlineSubsystemCreated.AddLambda([this](IOnlineSubsystem* NewSubsystem)
+		{
+			if (NewSubsystem == this)
+			{
+				FModuleManager::Get().LoadModule(TEXT("PlayFabParty"));
+			}
+		});
+#endif
 	}
 	else
 	{
@@ -852,6 +864,10 @@ bool FOnlineSubsystemGDK::Shutdown()
 		XblTaskQueue->CancelPendingTasksAndDestroyQueue();
 	}
 
+#ifdef UE_PLAYFAB_MATCHMAKING
+	FOnlineSubsystemDelegates::OnOnlineSubsystemCreated.Remove(OnOnlineSubsystemCreated);
+#endif
+
 	return true;
 }
 
diff --git a/Engine/Plugins/Online/Microsoft/OnlineSubsystemGDK/Source/Public/OnlineSubsystemGDK.h b/Engine/Plugins/Online/Microsoft/OnlineSubsystemGDK/Source/Public/OnlineSubsystemGDK.h
--- a/Engine/Plugins/Online/Microsoft/OnlineSubsystemGDK/Source/Public/OnlineSubsystemGDK.h
+++ b/Engine/Plugins/Online/Microsoft/OnlineSubsystemGDK/Source/Public/OnlineSubsystemGDK.h
@@ -376,6 +376,9 @@ private:
 	/** Stored delegate handle to remove the task later */
 	FDelegateHandle AppInitComplete;
 	FDelegateHandle UserLoginChanged;
+#ifdef UE_PLAYFAB_MATCHMAKING
+	FDelegateHandle OnOnlineSubsystemCreated;
+#endif
 
 	TOptional<XTaskQueueRegistrationToken> NetworkConnectivityHandle;
 	bool bShouldRestoreConnectivity = false;
diff --git a/Engine/Plugins/Online/Microsoft/PlayFabParty/PlayFabParty.uplugin b/Engine/Plugins/Online/Microsoft/PlayFabParty/PlayFabParty.uplugin
--- a/Engine/Plugins/Online/Microsoft/PlayFabParty/PlayFabParty.uplugin
+++ b/Engine/Plugins/Online/Microsoft/PlayFabParty/PlayFabParty.uplugin
@@ -10,6 +10,7 @@
 	"EnabledByDefault": false,
 	"RequiresBuildPlatform": false,
 	"HasExplicitPlatforms": true,
+	"ExplicitlyLoaded": true,
 	"SupportedTargetPlatforms": ["Win64"],
 	"Modules": [
 		{
