User remains in queue in the Lobby after a match is flushed.
* Console treats the network connectivity hint "Unknown" as "Disconnected" now to trigger the normal lost-internet-connection-code instead of silently ignoring the lost internet connection which can get the matchmaking into a bad state.

#Synced-CL 52242993

diff --git a/Engine/Plugins/Online/Microsoft/OnlineSubsystemGDK/Source/Private/OnlineSubsystemGDK.cpp b/Engine/Plugins/Online/Microsoft/OnlineSubsystemGDK/Source/Private/OnlineSubsystemGDK.cpp
--- a/Engine/Plugins/Online/Microsoft/OnlineSubsystemGDK/Source/Private/OnlineSubsystemGDK.cpp
+++ b/Engine/Plugins/Online/Microsoft/OnlineSubsystemGDK/Source/Private/OnlineSubsystemGDK.cpp
@@ -52,6 +52,17 @@ THIRD_PARTY_INCLUDES_END
 #include "Microsoft/HideMicrosoftPlatformTypes.h"
 // FOnlineSubsystemGDKModule
 
+namespace UE::OnlineSubsystemGDK
+{
+	bool bTreatConnectivityUnknownAsDisconnected = true;
+	static FAutoConsoleVariableRef CVarTreatConnectivityUnknownAsDisconnected(
+		TEXT("OnlineSubsystemGDK.TreatConnectivityUnknownAsDisconnected"),
+		bTreatConnectivityUnknownAsDisconnected,
+		TEXT("True (Default): Triggers OnConnectionStatusChanged delegate also for Unknown connectivity state.\n")
+		TEXT("False: Does not trigger OnConnectionStatusChanged delegate for Unknown connectivity state"),
+		ECVF_Default
+	);
+}
 
 bool FOnlineSubsystemGDK::IsEnabled() const
 {
@@ -449,7 +460,7 @@ void FOnlineSubsystemGDK::OnNetworkConnectivityHintChanged(const XNetworkingConn
 {
 	UE_LOG_ONLINE(Log, TEXT("ConnectivityHintChanged: %s"), LexToString(ConnectivityHint.connectivityLevel));
 
-	if (ConnectivityHint.connectivityLevel == XNetworkingConnectivityLevelHint::Unknown)
+	if (!UE::OnlineSubsystemGDK::bTreatConnectivityUnknownAsDisconnected && ConnectivityHint.connectivityLevel == XNetworkingConnectivityLevelHint::Unknown)
 	{
 		return;
 	}
@@ -467,7 +478,18 @@ void FOnlineSubsystemGDK::OnNetworkConnectivityHintChanged(const XNetworkingConn
 		}
 	}
 
-	ApplyNetworkConnectivityLevel(ConnectivityHint);
+	XNetworkingConnectivityHint ModifiedConnectivityHint = ConnectivityHint;
+	if (ModifiedConnectivityHint.connectivityLevel == XNetworkingConnectivityLevelHint::Unknown)
+	{
+		// bTreatConnectivityUnknownAsDisconnected is true, otherwise the early out in the top 
+		// of this function would have run. Setting connectivityLevel to "None" so the 
+		// OnConnectionStatusChanged delegate will be triggered from 
+		// ApplyNetworkConnectivityLevel -> FAsyncEventConnectionStatusChanged::TriggerDelegates
+		ModifiedConnectivityHint.connectivityLevel = XNetworkingConnectivityLevelHint::None;
+		UE_LOG_ONLINE(Log, TEXT("Modified ConnectivityHint from: %s to: %s"), LexToString(ConnectivityHint.connectivityLevel), LexToString(ModifiedConnectivityHint.connectivityLevel));
+	}
+
+	ApplyNetworkConnectivityLevel(ModifiedConnectivityHint);
 }
 
 void FOnlineSubsystemGDK::CleanupGDKContextForNetworkConnectivityLoss()
